Sequencer: fix p4ignore file

diff --git a/Engine/Plugins/Cameras/.p4ignore.txt b/Engine/Plugins/Cameras/.p4ignore.txt
--- a/Engine/Plugins/Cameras/.p4ignore.txt
+++ b/Engine/Plugins/Cameras/.p4ignore.txt
@@ -1 +1,2 @@
 */Binaries/*
+*/Intermediate/*
diff --git a/Engine/Plugins/MovieScene/.p4ignore.txt b/Engine/Plugins/MovieScene/.p4ignore.txt
--- a/Engine/Plugins/MovieScene/.p4ignore.txt
+++ b/Engine/Plugins/MovieScene/.p4ignore.txt
@@ -1 +1,2 @@
-*/Binaries/*
\ No newline at end of file
+*/Binaries/*
+*/Intermediate/*
\ No newline at end of file
